Merging //UE4/Dev-Main to Dev-Niagara (//UE4/Dev-Niagara)
Other folders

#tests n/a

diff --git a/Templates/TP_2DSideScroller/Config/DefaultEditor.ini b/Templates/TP_2DSideScroller/Config/DefaultEditor.ini
--- a/Templates/TP_2DSideScroller/Config/DefaultEditor.ini
+++ b/Templates/TP_2DSideScroller/Config/DefaultEditor.ini
@@ -1,5 +1,5 @@
 [UnrealEd.SimpleMap]
-SimpleMapName=/Game/2DSideScroller/Maps/2DSideScrollerExampleMap
+SimpleMapName=/Game/2DSideScrollerCPP/Maps/2DSideScrollerExampleMap
 
 [EditoronlyBP]
 bAllowClassAndBlueprintPinMatching=true
diff --git a/Templates/TP_2DSideScroller/Config/DefaultEngine.ini b/Templates/TP_2DSideScroller/Config/DefaultEngine.ini
--- a/Templates/TP_2DSideScroller/Config/DefaultEngine.ini
+++ b/Templates/TP_2DSideScroller/Config/DefaultEngine.ini
@@ -8,7 +8,7 @@ TwoPlayerSplitscreenLayout=Horizontal
 ThreePlayerSplitscreenLayout=FavorTop
 GameInstanceClass=/Script/Engine.GameInstance
 GameDefaultMap=/Game/2DSideScrollerCPP/Maps/2DSideScrollerExampleMap
-ServerDefaultMap=/Game/2DSideScroller/Maps/2DSideScrollerExampleMap
+ServerDefaultMap=/Game/2DSideScrollerCPP/Maps/2DSideScrollerExampleMap
 GlobalDefaultGameMode=/Script/TP_2DSideScroller.TP_2DSideScrollerGameMode
 GlobalDefaultServerGameMode=None
 
diff --git a/Templates/TP_HandheldARBP/manifest.json b/Templates/TP_HandheldARBP/manifest.json
--- a/Templates/TP_HandheldARBP/manifest.json
+++ b/Templates/TP_HandheldARBP/manifest.json
@@ -4,15 +4,15 @@
 	[
 		{
 			"Language": "en",
-			"Text": "HandheldAugmentedReality"
+			"Text": "Handheld AR"
 		}
 	],
 	"Description":
 	[
-		{
-			"Language": "en",
-			"Text": "This template project is set up for mobile deployment and contains logic for toggling AR mode, handling light estimation, displaying information on found planes and handling hit detection."
-		}
+    {
+      "Language": "en",
+      "Text": "A starting point for building Augmented Reality applications for Android and iOS devices. Includes runtime logic for toggling AR mode on and off, debug information regarding plane detection, example code for hit detection and handling light estimation."
+    }
 	],
 	"AssetTypes":
 	[
@@ -28,7 +28,7 @@
 			"Text": "AR, Mobile, Augmented Reality"
 		}
 	],
-	"ClassTypes": "GameMode, InputComponent",
+	"ClassTypes": "Pawn, GameMode, Widget, Actor",
 	"Category": "BlueprintFeature",
 	"FocusAsset": "/Game/Blueprints/GameFramework/BP_ARPawn.BP_ARPawn",
 	"Thumbnail": "TP_HandheldARBP.png",
diff --git a/Templates/TP_VirtualRealityBP/manifest.json b/Templates/TP_VirtualRealityBP/manifest.json
--- a/Templates/TP_VirtualRealityBP/manifest.json
+++ b/Templates/TP_VirtualRealityBP/manifest.json
@@ -1,51 +1,44 @@
 {
-	"Version": 1,
-	"Name":
-	[
-		{
-			"Language": "en",
-			"Text": "VirtualReality"
-		}
-	],
-	"Description":
-	[
-		{
-			"Language": "en",
-			"Text": "The virtual reality pack features two locomotion types for VR with and without motioncontrollers."
-		}
-	],
-	"AssetTypes":
-	[
-		{
-			"Language": "en",
-			"Text": "Static Mesh, Animation Blueprint, Physics Asset, Physics Material, Skeleton, Skeletal Mesh, Material, Material instance, Texture"
-		}
-	],
-	"SearchTags":
-	[
-		{
-			"Language": "en",
-			"Text": "VR,Virtual,Reality,Virtual Reality"
-		}
-	],
-	"ClassTypes": "GameMode, InputComponent",
-	"Category": "BlueprintFeature",
-	"FocusAsset": "/Game/VirtualRealityBP/VirtualReality.VirtualReality",
-	"Thumbnail": "TP_VirtualRealityBP.png",
-	"Screenshots":
-	[
-		"TP_VirtualRealityBP_Preview.png"
-	],
-	"SortKey": "VirtualRealityBP",
-	"AdditionalFeaturePacks":
-	[
-		{
-			"MountName": "Geometry",
-			"DetailLevels":
-			[
-				"Standard",
-				"High"
-			]
-		}
-	]
+  "Version": 1,
+  "Name": [
+    {
+      "Language": "en",
+      "Text": "Virtual Reality"
+    }
+  ],
+  "Description": [
+    {
+      "Language": "en",
+      "Text": "Blueprint Virtual Reality feature pack for Desktop and Playstation VR. The pack features two modes of teleport locomotion with separate pawns and maps for each. The MotionControllerMap lets players teleport around using motion controllers. The HMDLocomotionMap lets players teleport using either a Gamepad or point-and-teleport using just the HeadMountedDisplay."
+    }
+  ],
+  "AssetTypes": [
+    {
+      "Language": "en",
+      "Text": "Animation Sequence, Animation Montage, Sound Wave, Skeletal Mesh, Animation Blueprint, Skeleton, Static Mesh, Texture."
+    }
+  ],
+  "SearchTags": [
+    {
+      "Language": "en",
+      "Text": "VR,Virtual,Reality,Virtual Reality"
+    }
+  ],
+  "ClassTypes": "Character, GameMode, HUD, Actor, Character, GameMode, HUD, Actor",
+  "Category": "BlueprintFeature",
+  "FocusAsset": "/Game/VirtualRealityBP/VirtualReality.VirtualReality",
+  "Thumbnail": "TP_VirtualRealityBP.png",
+  "Screenshots": [
+    "TP_VirtualRealityBP_Preview.png"
+  ],
+  "SortKey": "VirtualRealityBP",
+  "AdditionalFeaturePacks": [
+    {
+      "MountName": "Geometry",
+      "DetailLevels": [
+        "Standard",
+        "High"
+      ]
+    }
+  ]
 }
