Removing Win32 from supported platforms

diff --git a/Templates/TP_VirtualRealityBP/TP_VirtualRealityBP.uproject b/Templates/TP_VirtualRealityBP/TP_VirtualRealityBP.uproject
--- a/Templates/TP_VirtualRealityBP/TP_VirtualRealityBP.uproject
+++ b/Templates/TP_VirtualRealityBP/TP_VirtualRealityBP.uproject
@@ -16,7 +16,6 @@
 			"Name": "OculusVR",
 			"Enabled": false,
 			"SupportedTargetPlatforms": [
-				"Win32",
 				"Win64",
 				"Android"
 			]
@@ -25,7 +24,6 @@
 			"Name": "OpenXR",
 			"Enabled": true,
 			"SupportedTargetPlatforms": [
-				"Win32",
 				"Win64",
 				"Android",
 				"HoloLens"
@@ -35,7 +33,6 @@
 			"Name": "SteamVR",
 			"Enabled": false,
 			"SupportedTargetPlatforms": [
-				"Win32",
 				"Win64",
 				"Linux"
 			]
