Disabled VR plugins by default in engine. Enabled on project base in TP_VirtualRealityBP template and in QAGame.

diff --git a/Engine/Plugins/Runtime/Oculus/OculusVR/OculusVR.uplugin b/Engine/Plugins/Runtime/Oculus/OculusVR/OculusVR.uplugin
--- a/Engine/Plugins/Runtime/Oculus/OculusVR/OculusVR.uplugin
+++ b/Engine/Plugins/Runtime/Oculus/OculusVR/OculusVR.uplugin
@@ -10,7 +10,7 @@
 	"DocsURL": "https://developer.oculus.com/documentation/unreal/latest/concepts/unreal-engine/",
 	"MarketplaceURL": "",
 	"SupportURL": "https://forums.oculusvr.com/developer/categories/unreal",
-	"EnabledByDefault": true,
+	"EnabledByDefault": false,
 	"CanContainContent": true,
 	"IsBetaVersion": false,
 	"Installed": false,
diff --git a/Engine/Plugins/Runtime/OpenXR/OpenXR.uplugin b/Engine/Plugins/Runtime/OpenXR/OpenXR.uplugin
--- a/Engine/Plugins/Runtime/OpenXR/OpenXR.uplugin
+++ b/Engine/Plugins/Runtime/OpenXR/OpenXR.uplugin
@@ -10,7 +10,7 @@
 	"DocsURL" : "",
 	"MarketplaceURL" : "",
 	"SupportURL" : "",
-	"EnabledByDefault" : true,
+	"EnabledByDefault" : false,
 	"CanContainContent" : true,
 	"IsBetaVersion" : false,
 	"Installed" : false,
diff --git a/Engine/Plugins/Runtime/OpenXREyeTracker/OpenXREyeTracker.uplugin b/Engine/Plugins/Runtime/OpenXREyeTracker/OpenXREyeTracker.uplugin
--- a/Engine/Plugins/Runtime/OpenXREyeTracker/OpenXREyeTracker.uplugin
+++ b/Engine/Plugins/Runtime/OpenXREyeTracker/OpenXREyeTracker.uplugin
@@ -10,7 +10,7 @@
 	"DocsURL": "",
 	"MarketplaceURL": "",
 	"SupportURL": "",
-	"EnabledByDefault": true,
+	"EnabledByDefault": false,
 	"CanContainContent": true,
 	"IsBetaVersion": false,
 	"Installed": false,
diff --git a/Engine/Plugins/Runtime/OpenXRHandTracking/OpenXRHandTracking.uplugin b/Engine/Plugins/Runtime/OpenXRHandTracking/OpenXRHandTracking.uplugin
--- a/Engine/Plugins/Runtime/OpenXRHandTracking/OpenXRHandTracking.uplugin
+++ b/Engine/Plugins/Runtime/OpenXRHandTracking/OpenXRHandTracking.uplugin
@@ -10,7 +10,7 @@
 	"DocsURL": "",
 	"MarketplaceURL": "",
 	"SupportURL": "",
-	"EnabledByDefault": true,
+	"EnabledByDefault": false,
 	"CanContainContent": true,
 	"IsBetaVersion": false,
 	"Installed": false,
diff --git a/Engine/Plugins/Runtime/Steam/SteamVR/SteamVR.uplugin b/Engine/Plugins/Runtime/Steam/SteamVR/SteamVR.uplugin
--- a/Engine/Plugins/Runtime/Steam/SteamVR/SteamVR.uplugin
+++ b/Engine/Plugins/Runtime/Steam/SteamVR/SteamVR.uplugin
@@ -10,7 +10,7 @@
 	"DocsURL": "",
 	"MarketplaceURL": "",
 	"SupportURL": "",
-	"EnabledByDefault": true,
+	"EnabledByDefault": false,
 	"CanContainContent": true,
 	"IsBetaVersion": false,
 	"Installed": false,
diff --git a/Templates/TP_ThirdPerson/TP_ThirdPerson.uproject b/Templates/TP_ThirdPerson/TP_ThirdPerson.uproject
--- a/Templates/TP_ThirdPerson/TP_ThirdPerson.uproject
+++ b/Templates/TP_ThirdPerson/TP_ThirdPerson.uproject
@@ -17,26 +17,6 @@
 			"TargetAllowList": [
 				"Editor"
 			]
-		},
-		{
-			"Name": "OpenXREyeTracker",
-			"Enabled": false,
-			"SupportedTargetPlatforms": [
-				"Win64",
-				"Linux",
-				"HoloLens",
-				"Android"
-			]
-		},
-		{
-			"Name": "OpenXRHandTracking",
-			"Enabled": false,
-			"SupportedTargetPlatforms": [
-				"Win64",
-				"Linux",
-				"HoloLens",
-				"Android"
-			]
 		}
 	]
 }
\ No newline at end of file
diff --git a/Templates/TP_ThirdPersonBP/TP_ThirdPersonBP.uproject b/Templates/TP_ThirdPersonBP/TP_ThirdPersonBP.uproject
--- a/Templates/TP_ThirdPersonBP/TP_ThirdPersonBP.uproject
+++ b/Templates/TP_ThirdPersonBP/TP_ThirdPersonBP.uproject
@@ -10,26 +10,6 @@
 			"TargetAllowList": [
 				"Editor"
 			]
-		},
-		{
-			"Name": "OpenXREyeTracker",
-			"Enabled": false,
-			"SupportedTargetPlatforms": [
-				"Win64",
-				"Linux",
-				"HoloLens",
-				"Android"
-			]
-		},
-		{
-			"Name": "OpenXRHandTracking",
-			"Enabled": false,
-			"SupportedTargetPlatforms": [
-				"Win64",
-				"Linux",
-				"HoloLens",
-				"Android"
-			]
 		}
 	]
 }
\ No newline at end of file
diff --git a/Templates/TP_VirtualRealityBP/TP_VirtualRealityBP.uproject b/Templates/TP_VirtualRealityBP/TP_VirtualRealityBP.uproject
--- a/Templates/TP_VirtualRealityBP/TP_VirtualRealityBP.uproject
+++ b/Templates/TP_VirtualRealityBP/TP_VirtualRealityBP.uproject
@@ -24,12 +24,39 @@
 		},
 		{
 			"Name": "SteamVR",
-			"Enabled": false,
+			"Enabled": true,
 			"SupportedTargetPlatforms": [
-				"Win32",
 				"Win64",
 				"Linux"
 			]
+		},
+		{
+			"Name": "OculusVR",
+			"Enabled": true,
+			"SupportedTargetPlatforms": [
+				"Win64",
+				"Android"
+			]
+		},
+		{
+			"Name": "OpenXREyeTracker",
+			"Enabled": true,
+			"SupportedTargetPlatforms": [
+				"Win64",
+				"Linux",
+				"HoloLens",
+				"Android"
+			]
+		},
+		{
+			"Name": "OpenXRHandTracking",
+			"Enabled": true,
+			"SupportedTargetPlatforms": [
+				"Win64",
+				"Linux",
+				"HoloLens",
+				"Android"
+			]
 		}
 	],
 	"TargetPlatforms": [
