XR | VR Template | Disabling SteamVR plugin in VR Template as it was enabled in the big VR plugin purge

diff --git a/Templates/TP_VirtualRealityBP/TP_VirtualRealityBP.uproject b/Templates/TP_VirtualRealityBP/TP_VirtualRealityBP.uproject
--- a/Templates/TP_VirtualRealityBP/TP_VirtualRealityBP.uproject
+++ b/Templates/TP_VirtualRealityBP/TP_VirtualRealityBP.uproject
@@ -4,10 +4,6 @@
 	"Category": "",
 	"Description": "",
 	"Plugins": [
-		{
-			"Name": "XRVisualization",
-			"Enabled": false
-		},
 		{
 			"Name": "ResonanceAudio",
 			"Enabled": true
@@ -22,14 +18,6 @@
 				"HoloLens"
 			]
 		},
-		{
-			"Name": "SteamVR",
-			"Enabled": true,
-			"SupportedTargetPlatforms": [
-				"Win64",
-				"Linux"
-			]
-		},
 		{
 			"Name": "OculusVR",
 			"Enabled": true,
@@ -37,26 +25,6 @@
 				"Win64",
 				"Android"
 			]
-		},
-		{
-			"Name": "OpenXREyeTracker",
-			"Enabled": true,
-			"SupportedTargetPlatforms": [
-				"Win64",
-				"Linux",
-				"HoloLens",
-				"Android"
-			]
-		},
-		{
-			"Name": "OpenXRHandTracking",
-			"Enabled": true,
-			"SupportedTargetPlatforms": [
-				"Win64",
-				"Linux",
-				"HoloLens",
-				"Android"
-			]
 		}
 	],
 	"TargetPlatforms": [
