XR | VR Template | Enabled HoloLens as a platform

diff --git a/Templates/TP_VirtualRealityBP/TP_VirtualRealityBP.uproject b/Templates/TP_VirtualRealityBP/TP_VirtualRealityBP.uproject
--- a/Templates/TP_VirtualRealityBP/TP_VirtualRealityBP.uproject
+++ b/Templates/TP_VirtualRealityBP/TP_VirtualRealityBP.uproject
@@ -38,6 +38,7 @@
 		"WindowsNoEditor",
 		"Switch",
 		"XboxOne",
-		"XSX"
+		"XSX",
+		"HoloLens"
 	]
 }
\ No newline at end of file
