Updating title/description translation for the VR Template.

diff --git a/Templates/TP_VirtualRealityBP/Config/TemplateDefs.ini b/Templates/TP_VirtualRealityBP/Config/TemplateDefs.ini
--- a/Templates/TP_VirtualRealityBP/Config/TemplateDefs.ini
+++ b/Templates/TP_VirtualRealityBP/Config/TemplateDefs.ini
@@ -5,18 +5,19 @@ Categories=Games
 LocalizedDisplayNames=(Language="en",Text="Virtual Reality")
 LocalizedDescriptions=(Language="en",Text="Blueprint Virtual Reality Template using OpenXR for Desktop, Console, and Mobile. The template features teleport locomotion, snap turning, grabbable objects, grabbable and interactive objects, Google Resonance spatialized audio, and a VR Spectator Camera.")
 
-LocalizedDisplayNames=(Language="ko",Text="가상 현실")
-LocalizedDescriptions=(Language="ko",Text="데스크톱 및 Playstation VR 용 블루프린트 가상 현실 템플릿입니다. 각 플랫폼 별 전용 폰과 맵은 물론 두 가지 텔레포트 이동 방식이 포함되어 있습니다. 플레이어는 MotionControllerMap 에서 모션 컨트롤러를 사용하여 텔레포트할 수 있으며, HMDLocomotionMap 에서 게임패드 아니면 HeadMountedDisplay 만으로 이동할 곳을 지정하는 방식으로 텔레포트할 수 있습니다.")
+LocalizedDisplayNames=(Language="ko",Text="가상현실")
+LocalizedDescriptions=(Language="ko",Text="데스크톱, 콘솔 및 모바일용 OpenXR을 사용하는 블루프린트 VR 템플릿으로, 텔레포트 로코모션, 스냅 회전, 손으로 잡는 오브젝트, 잡을 수 있는 인터랙티브 오브젝트, Google Resonance에서 공간화된 오디오, 그리고 VR 관람자 카메라를 포함하고 있습니다.")
 
 LocalizedDisplayNames=(Language="ja",Text="バーチャルリアリティ")
-LocalizedDescriptions=(Language="ja",Text="デスクトップ VR と Playstation VR 向け、ブループリントのバーチャルリアリティテンプレートです。テンプレートには二種類のテレポートロコモーションモードがあります。それぞれに ポーンとマップが用意されています。MotionControllerMap ではモーションコントローラーを使ったテレポートが可能です。HMDLocomotionMap ではゲームパッドを使ったテレポート、及び、ヘッドマウントディスプレイを使ったポイントアンドテレポートが可能です。")
+LocalizedDescriptions=(Language="ja",Text="デスクトップ、コンソール、およびモバイル用のOpenXRを使用したブループリントバーチャルリアリティテンプレート。テンプレートには、テレポートによる移動、スナップターニング、掴めるオブジェクト、掴めてインタラクティブなオブジェクト、Google Resonanceの空間化されたオーディオ、VR Spectatorカメラが含まれています。")
 
-LocalizedDisplayNames=(Language="zh-Hans",Text="虚拟现实应用")
+LocalizedDisplayNames=(Language="zh-Hans",Text="虚拟现实")
 LocalizedDescriptions=(Language="zh-Hans",Text="蓝图虚拟现实模板基于OpenXR实现，用于为面向台式电脑、主机端以及移动端的VR设备开发项目。该模板默认实现了玩家传送、玩家旋转、物体拾取、物体交互、Google Resonance空间化音频，以及VR观察视角等功能。")
 
 ClassTypes=Pawn, GameMode, Blueprint Interface, Actor
 AssetTypes=Sound Cue, Haptic Feedback Effect Curve, Niagara System.
 
+
 FoldersToIgnore=Binaries
 FoldersToIgnore=Build
 FoldersToIgnore=Intermediate
diff --git a/Templates/TP_VirtualRealityBP/manifest.json b/Templates/TP_VirtualRealityBP/manifest.json
--- a/Templates/TP_VirtualRealityBP/manifest.json
+++ b/Templates/TP_VirtualRealityBP/manifest.json
@@ -7,7 +7,7 @@
 		},
 		{
 			"Language": "ko",
-			"Text": "가상 현실"
+			"Text": "가상현실"
 		},
 		{
 			"Language": "ja",
@@ -15,7 +15,7 @@
 		},
 		{
 			"Language": "zh-Hans",
-			"Text": "虚拟现实应用"
+			"Text": "虚拟现实"
 		}
 	],
 	"Description": [
@@ -25,27 +25,44 @@
 		},
 		{
 			"Language": "ko",
-			"Text": "데스크톱 및 Playstation VR 용 블루프린트 가상 현실 템플릿입니다. 각 플랫폼 별 전용 폰과 맵은 물론 두 가지 텔레포트 이동 방식이 포함되어 있습니다. 플레이어는 MotionControllerMap 에서 모션 컨트롤러를 사용하여 텔레포트할 수 있으며, HMDLocomotionMap 에서 게임패드 아니면 HeadMountedDisplay 만으로 이동할 곳을 지정하는 방식으로 텔레포트할 수 있습니다."
+			"Text": "데스크톱, 콘솔 및 모바일용 OpenXR을 사용하는 블루프린트 VR 템플릿으로, 텔레포트 로코모션, 스냅 회전, 손으로 잡는 오브젝트, 잡을 수 있는 인터랙티브 오브젝트, Google Resonance에서 공간화된 오디오, 그리고 VR 관람자 카메라를 포함하고 있습니다."
 		},
 		{
 			"Language": "ja",
-			"Text": "デスクトップ VR と Playstation VR 向け、ブループリントのバーチャルリアリティテンプレートです。テンプレートには二種類のテレポートロコモーションモードがあります。それぞれに ポーンとマップが用意されています。MotionControllerMap ではモーションコントローラーを使ったテレポートが可能です。HMDLocomotionMap ではゲームパッドを使ったテレポート、及び、ヘッドマウントディスプレイを使ったポイントアンドテレポートが可能です。"
+			"Text": "デスクトップ、コンソール、およびモバイル用のOpenXRを使用したブループリントバーチャルリアリティテンプレート。テンプレートには、テレポートによる移動、スナップターニング、掴めるオブジェクト、掴めてインタラクティブなオブジェクト、Google Resonanceの空間化されたオーディオ、VR Spectatorカメラが含まれています。"
 		},
 		{
 			"Language": "zh-Hans",
-			"Text": "蓝图虚拟现实应用模板面向台式PC以及Playstation VR设备。模板包含2中传送移动模式，分别拥有独立的Pawn和地图。MotionControllerMap可以让玩家使用体感手柄四处传送，而HMDLocomotionMap可以让玩家使用游戏手柄或头显设备进行“指向-传送”式的传送。"
+			"Text": "蓝图虚拟现实模板基于OpenXR实现，用于为面向台式电脑、主机端以及移动端的VR设备开发项目。该模板默认实现了玩家传送、玩家旋转、物体拾取、物体交互、Google Resonance空间化音频，以及VR观察视角等功能。"
 		}
 	],
 	"AssetTypes": [
 		{
 			"Language": "en",
 			"Text": "Sound Cue, Haptic Feedback Effect Curve, Niagara System."
+		}，
+		{
+			"Language": "ko",
+			"Text": "사운드 큐, 햅틱 피드백 이펙트 커브, 나이아가라 시스템"
+
 		}
 	],
 	"SearchTags": [
 		{
 			"Language": "en",
 			"Text": "VR,Virtual,Reality,Virtual Reality"
+		},
+		{
+			"Language": "ko",
+			"Text": "VR, 가상, 현실, 가상현실"
+		},
+		{
+			"Language": "ja",
+			"Text": "VR,バーチャル,リアリティ,バーチャルリアリティ"
+		},
+		{
+			"Language": "zh-Hans",
+			"Text": "VR,Virtual,Reality,虚拟现实"
 		}
 	],
 	"ClassTypes": "Pawn, GameMode, Blueprint Interface, Actor",
