XR | VR Template | Disabled OculusVR plugin as it's being deprecated in 5.1

diff --git a/Templates/TP_VirtualRealityBP/TP_VirtualRealityBP.uproject b/Templates/TP_VirtualRealityBP/TP_VirtualRealityBP.uproject
--- a/Templates/TP_VirtualRealityBP/TP_VirtualRealityBP.uproject
+++ b/Templates/TP_VirtualRealityBP/TP_VirtualRealityBP.uproject
@@ -17,14 +17,6 @@
 				"Android",
 				"HoloLens"
 			]
-		},
-		{
-			"Name": "OculusVR",
-			"Enabled": true,
-			"SupportedTargetPlatforms": [
-				"Win64",
-				"Android"
-			]
 		}
 	],
 	"TargetPlatforms": [
