Updated the Feature Pack for the VR Template. This will update it with all of the changes that's been done since 5.0, but also fix UE-157808

diff --git a/Templates/TP_VirtualRealityBP/manifest.json b/Templates/TP_VirtualRealityBP/manifest.json
--- a/Templates/TP_VirtualRealityBP/manifest.json
+++ b/Templates/TP_VirtualRealityBP/manifest.json
@@ -40,7 +40,7 @@
 		{
 			"Language": "en",
 			"Text": "Sound Cue, Haptic Feedback Effect Curve, Niagara System."
-		}，
+		},
 		{
 			"Language": "ko",
 			"Text": "사운드 큐, 햅틱 피드백 이펙트 커브, 나이아가라 시스템"
