XR | VR Template | Enabled OpenXRHandTracker and OpenXREyeTracker-plugins

diff --git a/Templates/TP_VirtualRealityBP/TP_VirtualRealityBP.uproject b/Templates/TP_VirtualRealityBP/TP_VirtualRealityBP.uproject
--- a/Templates/TP_VirtualRealityBP/TP_VirtualRealityBP.uproject
+++ b/Templates/TP_VirtualRealityBP/TP_VirtualRealityBP.uproject
@@ -10,8 +10,25 @@
 			"SupportedTargetPlatforms": [
 				"Win64",
 				"Linux",
-				"Android",
-				"HoloLens"
+				"Android"
+			]
+		},
+		{
+			"Name": "OpenXREyeTracker",
+			"Enabled": true,
+			"SupportedTargetPlatforms": [
+				"Win64",
+				"Linux",
+				"Android"
+			]
+		},
+		{
+			"Name": "OpenXRHandTracking",
+			"Enabled": true,
+			"SupportedTargetPlatforms": [
+				"Win64",
+				"Linux",
+				"Android"
 			]
 		}
 	],
