XR | VR Template | Adding visionOS as a supported platform

diff --git a/Templates/TP_VirtualRealityBP/TP_VirtualRealityBP.uproject b/Templates/TP_VirtualRealityBP/TP_VirtualRealityBP.uproject
--- a/Templates/TP_VirtualRealityBP/TP_VirtualRealityBP.uproject
+++ b/Templates/TP_VirtualRealityBP/TP_VirtualRealityBP.uproject
@@ -44,6 +44,6 @@
 		"Switch",
 		"XboxOne",
 		"XSX",
-		"HoloLens"
+		"VisionOS"
 	]
 }
\ No newline at end of file
