XR | VR Template | Adding visionOS to SupportedTargetPlatforms for OpenXR and OpenXRHandTracking

diff --git a/Templates/TP_VirtualRealityBP/TP_VirtualRealityBP.uproject b/Templates/TP_VirtualRealityBP/TP_VirtualRealityBP.uproject
--- a/Templates/TP_VirtualRealityBP/TP_VirtualRealityBP.uproject
+++ b/Templates/TP_VirtualRealityBP/TP_VirtualRealityBP.uproject
@@ -10,7 +10,8 @@
 			"SupportedTargetPlatforms": [
 				"Win64",
 				"Linux",
-				"Android"
+				"Android",
+				"VisionOS"
 			]
 		},
 		{
@@ -28,7 +29,8 @@
 			"SupportedTargetPlatforms": [
 				"Win64",
 				"Linux",
-				"Android"
+				"Android",
+				"VisionOS"
 			]
 		}
 	],
