BP_VirtualRealityBP - Fixed incorrect resource dependency on feature pack.
#rnx

diff --git a/Templates/TP_VirtualRealityBP/manifest.json b/Templates/TP_VirtualRealityBP/manifest.json
--- a/Templates/TP_VirtualRealityBP/manifest.json
+++ b/Templates/TP_VirtualRealityBP/manifest.json
@@ -81,7 +81,7 @@
 			]
 		},
 		{
-			"MountName": "FPWeapon",
+			"MountName": "Weapons",
 			"DetailLevels": [
 				"Standard"
 			]
