VR Template | Renaming Blueprints to match the style guide of new templates

#virtualized

diff --git a/Templates/TP_VirtualRealityBP/manifest.json b/Templates/TP_VirtualRealityBP/manifest.json
--- a/Templates/TP_VirtualRealityBP/manifest.json
+++ b/Templates/TP_VirtualRealityBP/manifest.json
@@ -67,7 +67,7 @@
 	],
 	"ClassTypes": "Pawn, GameMode, Blueprint Interface, Actor",
 	"Category": "BlueprintFeature",
-	"FocusAsset": "/Game/VRTemplate/Blueprints/VRPawn.VRPawn",
+	"FocusAsset": "/Game/VRTemplate/Blueprints/BP_XRPawn.BP_XRPawn",
 	"Thumbnail": "TP_VirtualRealityBP.png",
 	"Screenshots": [
 		"TP_VirtualRealityBP_Preview.png"
